Remove hard-coded engine versions from Magic Leap plugins to fix CIS warnings. They will always be compatible with the current engine version, so do not need to declare support for a specific version.

diff --git a/Engine/Plugins/Lumin/MagicLeap/MagicLeap.uplugin b/Engine/Plugins/Lumin/MagicLeap/MagicLeap.uplugin
--- a/Engine/Plugins/Lumin/MagicLeap/MagicLeap.uplugin
+++ b/Engine/Plugins/Lumin/MagicLeap/MagicLeap.uplugin
@@ -5,7 +5,6 @@
 	"FriendlyName" : "MagicLeap",
 	"CreatedBy" : "Magic Leap, Inc.",
 	"CreatedByURL" : "http://magicleap.com",
-	"EngineVersion" : "4.20.0",
 	"Description" : "Support for the Magic Leap Head Mounted Display.",
 	"Category" : "Augmented Reality",
 	"EnabledByDefault" : true,
diff --git a/Engine/Plugins/Lumin/MagicLeapAnalytics/MagicLeapAnalytics.uplugin b/Engine/Plugins/Lumin/MagicLeapAnalytics/MagicLeapAnalytics.uplugin
--- a/Engine/Plugins/Lumin/MagicLeapAnalytics/MagicLeapAnalytics.uplugin
+++ b/Engine/Plugins/Lumin/MagicLeapAnalytics/MagicLeapAnalytics.uplugin
@@ -5,7 +5,6 @@
 	"FriendlyName" : "Magic Leap Analytics",
 	"CreatedBy" : "Magic Leap, Inc.",
 	"CreatedByURL" : "http://magicleap.com",
-	"EngineVersion" : "4.20.0",
 	"Description" : "Writes analytic API calls to Magic Leap cloud server.",
 	"Category" : "Analytics",
 	"EnabledByDefault" : false,
diff --git a/Engine/Plugins/Lumin/MagicLeapEmulator/MagicLeapEmulator.uplugin b/Engine/Plugins/Lumin/MagicLeapEmulator/MagicLeapEmulator.uplugin
--- a/Engine/Plugins/Lumin/MagicLeapEmulator/MagicLeapEmulator.uplugin
+++ b/Engine/Plugins/Lumin/MagicLeapEmulator/MagicLeapEmulator.uplugin
@@ -5,7 +5,6 @@
 	"FriendlyName" : "MagicLeapEmulator",
 	"CreatedBy" : "Magic Leap, Inc.",
 	"CreatedByURL" : "http://magicleap.com",
-	"EngineVersion" : "4.20.0",
 	"Description" : "Support for the emulating a Magic Leap Head Mounted Display with an alternate display device.",
 	"Category" : "Augmented Reality",
 	"EnabledByDefault" : true,
diff --git a/Engine/Plugins/Lumin/MagicLeapMedia/MagicLeapMedia.uplugin b/Engine/Plugins/Lumin/MagicLeapMedia/MagicLeapMedia.uplugin
--- a/Engine/Plugins/Lumin/MagicLeapMedia/MagicLeapMedia.uplugin
+++ b/Engine/Plugins/Lumin/MagicLeapMedia/MagicLeapMedia.uplugin
@@ -5,7 +5,6 @@
 	"FriendlyName" : "Magic Leap Media Player",
 	"CreatedBy" : "Magic Leap, Inc.",
 	"CreatedByURL" : "http://magicleap.com",
-	"EngineVersion" : "4.20.0",
 	"Description" : "Implements a media player using the Magic Leap Media library.",
 	"Category" : "Media Players",
 	"DocsURL" : "",
diff --git a/Engine/Plugins/Lumin/MagicLeapScreens/MagicLeapScreens.uplugin b/Engine/Plugins/Lumin/MagicLeapScreens/MagicLeapScreens.uplugin
--- a/Engine/Plugins/Lumin/MagicLeapScreens/MagicLeapScreens.uplugin
+++ b/Engine/Plugins/Lumin/MagicLeapScreens/MagicLeapScreens.uplugin
@@ -5,7 +5,6 @@
 	"FriendlyName" : "Magic Leap Screens",
 	"CreatedBy" : "Magic Leap, Inc.",
 	"CreatedByURL" : "http://magicleap.com",
-	"EngineVersion" : "4.20.0",
 	"Description" : "Support for the Magic Leap Screens and Magic TV Apps.",
 	"Category" : "Movie Players",
 	"EnabledByDefault" : false,
